Rename/move file(s)
Volumetrics:  Moving from NotForLicensees to Experimental in preparation for wider sharing of volumetrics plugin.

diff --git a/Engine/Plugins/Experimental/Volumetrics/Shaders/Private/TestCustomNode.ush b/Engine/Plugins/Experimental/Volumetrics/Shaders/Private/TestCustomNode.ush
new file mode 100644
--- /dev/null
+++ b/Engine/Plugins/Experimental/Volumetrics/Shaders/Private/TestCustomNode.ush
@@ -0,0 +1,12 @@
+#pragma once
+
+struct StructExample
+{
+    float AddValues(float inA, float inB)
+    {
+        return inA + inB;
+    }
+};
+
+StructExample MyStruct;
+return MyStruct.AddValues(A, B); 
\ No newline at end of file
diff --git a/Engine/Plugins/Experimental/Volumetrics/Shaders/Private/VolumetricCloudFunctions.ush b/Engine/Plugins/Experimental/Volumetrics/Shaders/Private/VolumetricCloudFunctions.ush
new file mode 100644
--- /dev/null
+++ b/Engine/Plugins/Experimental/Volumetrics/Shaders/Private/VolumetricCloudFunctions.ush
@@ -0,0 +1,296 @@
+#pragma once
+
+#define UseDetailTexture 1
+#define NumShadowSteps 32
+
+//input RO 
+//input RD
+//input PlanetRadius
+//input SD : SceneDepth
+
+//input LayerAParams : input v2  R=Layer Altitude G=Layer Thickness
+//input LayerBParams : input v2  R=Layer Altitude G=Layer Thickness
+
+//input MaxSteps
+//input StepSize
+//input Tex : Cloud Volume Texture
+
+//input CloudDivisor : texture scale
+//input LowerFalloff : falloff curve for bottom of layer
+//input UpperFalloff : falloff curve for top of layer
+//input CloudDensity
+//input ShadowDensity
+
+const float TanHalfFOV = 1;
+const float3 SO = float3(0, 0, PlanetRadius);
+RO += SO;
+
+float accumdist = 0;
+float curdensity = 0;
+float transmittance = 1;
+
+float shadowstepsize = StepSize * min( 8, 1 / abs(LightVector.z));;
+StepSize *= min( 4, 1 / abs(RD.z));
+
+float initialstepsize = StepSize;
+
+
+
+CloudDensity = (CloudDensity / 1000);
+
+float cloud1mip = 0; //temp to store calculated mip level
+float cloud2mip = 0;
+float skipdist  = 0;  //track skipping empty space for Mip Calculation
+
+float cursample, lastsample = 0; //store previous sample to average slices
+
+struct SkyScene
+{
+	
+    float3 SphereIntersection(float Radius)
+    {
+        float B = dot(RD, RO);
+        float C = dot(RO, RO) - Radius * Radius;
+	
+        float t0 = -B - sqrt(B * B - C);
+        float t1 = -B + sqrt(B * B - C);
+		
+        t0 = max(t0, 0);
+        t1 = (t1 > 0) ? min(t1, SD) : t1;
+	
+        return float3(t0, t1, max(0, t1 - t0));
+    }
+	
+	float GetMipLevel(float RayDistance, float RepeatSize)
+    {
+        float TexelsPerPixel = ((RayDistance * TanHalfFOV * 2) / RepeatSize) * (VolumeResolution / max(ViewportSize.x, ViewportSize.y));
+        return max(0, log2(TexelsPerPixel));
+
+    }
+	
+    float GetNormalizedRingDistance(float3 Pos, float ringnoise)
+    {
+        float radialdist = length(Pos);
+        radialdist -= PlanetRadius + LayerAParams.x;
+        radialdist /= LayerAParams.y;
+		
+		radialdist += ringnoise;
+		
+        return radialdist;
+    }
+	
+    float GetLayerFalloff(float3 Pos, float layernoise)
+    {
+        float ringdist = GetNormalizedRingDistance(Pos, layernoise);
+		
+        return (1 - exp(-ringdist * LowerFalloff)) * (1 - exp(-(1 - ringdist) * UpperFalloff));
+    }
+	
+    float SampleCloudLayer(float3 Pos)
+    {
+		
+        float volsample = Tex.SampleLevel(TexSampler, (Pos) / CloudDivisor, 0 + Mip1).r;
+		
+		#if UseDetailTexture
+			float detailsample = Tex.SampleLevel(TexSampler, (Pos / Cloud2Divisor) + Pan, 0 + Mip2).r;	
+			volsample = volsample + ((detailsample - 0.5) * Cloud2Density) - Bias;
+		#else
+			volsample = volsample - Bias;
+		#endif
+		
+		float falloff = GetLayerFalloff(Pos, (volsample -1)* LayerBreakup);
+        volsample -= (1 / max( 0.01,falloff )) * Bias;
+		volsample *= CloudDensity;
+        volsample *= saturate(falloff);
+
+		
+        return volsample;
+    }
+	
+    float4 RayMarch(float3 Pos, float3 Dir, int Steps, float MaxDist)
+    {
+        float accum = 0;
+        float dist = 0;
+        float3 lightenergy = 0;
+	
+        for (int i = 0; i < Steps; i++)
+        {
+            StepSize *= 1.02;
+		
+            cursample = SampleCloudLayer(Pos);
+			
+            dist += StepSize;
+			
+			if( cursample > 0)
+			{
+			
+				accum += cursample;
+			
+				float3 lpos = Pos;
+				float shadowdist = 0;
+				
+				for (int shadowstep = 0; shadowstep < NumShadowSteps; shadowstep++)
+				{
+					lpos += LightVector * shadowstepsize;
+					float lsample = SampleCloudLayer(lpos);
+					
+					shadowdist += max(0,lsample);
+					
+				}
+				
+				curdensity = 1 - exp(-max(0,cursample) * StepSize);
+				lightenergy += exp(-shadowdist * ShadowDensity * shadowstepsize) * curdensity * SunColor * transmittance;
+				transmittance *= 1 - curdensity;
+				
+				
+
+				if(transmittance < 0.001)
+				{
+					//transmittance = 0;
+					//return float4(lightenergy, accum);
+				}
+				
+				#if 1 //Sky Lighting
+				shadowdist = 0;
+				lpos = Pos + normalize(Pos) * StepSize;
+				float lsample = SampleCloudLayer(lpos);
+				shadowdist += lsample;
+				lpos = Pos + normalize(Pos) * StepSize * 2;
+				lsample = SampleCloudLayer(lpos);
+				shadowdist += lsample;
+				shadowdist = max(0,shadowdist);
+				
+				lightenergy += exp(-shadowdist * AmbientDensity * StepSize) * curdensity * SkyColor * transmittance;
+				
+				#endif
+				
+				#if 1 // ground bounce
+				shadowdist = 0;
+				lpos = Pos - normalize(Pos) * StepSize * 2;
+				lsample = SampleCloudLayer(lpos);
+				shadowdist += lsample;
+				
+				lightenergy += exp(-shadowdist * AmbientDensity) * curdensity * dot(1, SkyColor) * GroundColor * transmittance;
+				
+				#endif
+			}
+			
+			if (dist > MaxDist)
+			{
+				//	Pos += Dir * fmod(StepSize, MaxDist);
+				//	cursample = PseudoVolumeTexture(Tex, TexSampler, Pos / CloudDivisor, 16, 256).r;
+				//	//accum += cursample * fmod(StepSize, MaxDist);
+					return float4(lightenergy, accum);
+			}
+			
+            Pos += Dir * StepSize;
+			StepSize = max(initialstepsize, ((-cursample) * CloudDivisor*raycastdist)/ CloudDensity * (dist / MaxDist ));
+			
+			accumdist += StepSize * transmittance;
+        }
+		
+        return float4(lightenergy, transmittance);
+    }
+	
+    float4 Render()
+    {
+        float LayerABottom = PlanetRadius + LayerAParams.x;
+        float LayerATop = LayerABottom + LayerAParams.y;
+		
+        float3 PlanetOuter = SphereIntersection(PlanetRadius);
+		
+		//Use Planetary Distance beyond depth precision limits
+        if (SD > 650000 && PlanetOuter.x > 0)
+        {
+            SD = PlanetOuter.x;
+        }
+		
+        float3 LayerAOuter = SphereIntersection(LayerATop);
+        float3 LayerAInner = SphereIntersection(LayerABottom);
+		
+        float cloudentry = LayerAOuter.x;
+        float cloudexit = LayerAOuter.y;
+		
+		//extract rays that exit and re-enter the cloud layer
+		//Is Camera outisde the Inner Cloud layer
+        float mask = (LayerAInner.x > 0) ? 1 : 0;
+		//ignore rays that hit the planet surface
+        mask *= (PlanetOuter.x > 0) ? 0 : 1;
+		
+		//Ray exits and re-enters, stop the initial ray at that first exit spot, so we can skip empty space inbetween easily
+        if (mask == 1)
+        {
+            cloudexit = min(LayerAInner.x, SD);
+        }
+	
+		
+		// Ray Hits the Inner Shell of Cloud layer
+        if (LayerAInner.z > 0)
+        {
+            cloudexit = min(LayerAInner.x, SD);
+        }
+		
+		//Camera is inside the bottom/inner cloud radius, advacnce ray to inner shell
+        if (LayerAInner.x == 0 && LayerAInner.z > 0)
+        {
+		
+			//LayerAInner.y 
+		
+		
+            cloudentry = LayerAInner.y;
+            cloudexit = LayerAOuter.y;
+        }
+		
+        float4 clouds = 0;
+        float clouddist = cloudexit - cloudentry;
+        float tracedist = 0;
+		
+        if (clouddist > 0)
+        {
+			accumdist = cloudentry;
+			//StepSize = clouddist / MaxSteps;
+			
+			
+			float calcstep = ceil(clouddist / StepSize);
+			
+			
+            clouds = RayMarch(RO + cloudentry * RD, RD, MaxSteps, clouddist);
+        }
+		
+        if (mask == 1 && LayerAInner.y < SD)
+        {
+            cloudentry = min(LayerAInner.y, SD);
+            cloudexit = min(LayerAOuter.y, SD);
+			
+			skipdist = cloudentry;
+			
+            tracedist = cloudexit - cloudentry;
+			//StepSize = tracedist / MaxSteps;
+            clouds += RayMarch(RO + cloudentry * RD, RD, MaxSteps, tracedist);
+			
+            clouddist += tracedist;
+            clouddist += LayerAInner.z;
+			
+        }
+		
+        clouddist /= 2 * PlanetRadius;
+	
+		#if 0
+			accumdist /= 10000;
+			float aerialp = (1-exp(-accumdist * AtmosphereDensity));
+			transmittance = lerp(transmittance, 1, aerialp);
+			clouds.rgb *= 1-aerialp;
+			clouds.rgb += (SunColor + SkyColor) * aerialp * ExtraFog;
+		#else
+			float3 appos = RO + (accumdist * RD) - SO;
+			float3 ap = MaterialExpressionSkyAtmosphereAerialPerspective(Parameters, appos);
+			clouds.rgb += ap * AtmosphereDensity;
+			
+		#endif
+		
+		
+		
+        return float4(clouds.rgb, transmittance);
+    }
+	
+};
\ No newline at end of file
diff --git a/Engine/Plugins/Experimental/Volumetrics/Source/Volumetrics/Private/Volumetrics.cpp b/Engine/Plugins/Experimental/Volumetrics/Source/Volumetrics/Private/Volumetrics.cpp
new file mode 100644
--- /dev/null
+++ b/Engine/Plugins/Experimental/Volumetrics/Source/Volumetrics/Private/Volumetrics.cpp
@@ -0,0 +1,32 @@
+// Copyright 1998-2019 Epic Games, Inc. All Rights Reserved.
+
+#include "CoreMinimal.h"
+#include "Modules/ModuleManager.h"
+#include "IVolumetrics.h"
+#include "Interfaces/IPluginManager.h"
+#include "Misc/Paths.h"
+#include "ShaderCore.h"
+
+
+class FVolumetrics : public IVolumetrics
+{
+	/** IModuleInterface implementation */
+	virtual void StartupModule() override;
+	virtual void ShutdownModule() override;
+};
+
+IMPLEMENT_MODULE( FVolumetrics, Volumetrics )
+
+
+void FVolumetrics::StartupModule()
+{
+	FString PluginShaderDir = FPaths::Combine(IPluginManager::Get().FindPlugin(TEXT("Volumetrics"))->GetBaseDir(), TEXT("Shaders"));
+	AddShaderSourceDirectoryMapping(TEXT("/Plugin/Volumetrics"), PluginShaderDir);
+}
+
+
+void FVolumetrics::ShutdownModule()
+{
+	// This function may be called during shutdown to clean up your module.  For modules that support dynamic reloading,
+	// we call this function before unloading the module.
+}
diff --git a/Engine/Plugins/Experimental/Volumetrics/Source/Volumetrics/Public/IVolumetrics.h b/Engine/Plugins/Experimental/Volumetrics/Source/Volumetrics/Public/IVolumetrics.h
new file mode 100644
--- /dev/null
+++ b/Engine/Plugins/Experimental/Volumetrics/Source/Volumetrics/Public/IVolumetrics.h
@@ -0,0 +1,39 @@
+// Copyright 1998-2019 Epic Games, Inc. All Rights Reserved.
+
+#pragma once
+
+#include "CoreMinimal.h"
+#include "Modules/ModuleInterface.h"
+#include "Modules/ModuleManager.h"
+
+
+/**
+ * The public interface to this module
+ */
+class IVolumetrics : public IModuleInterface
+{
+
+public:
+
+	/**
+	 * Singleton-like access to this module's interface.  This is just for convenience!
+	 * Beware of calling this during the shutdown phase, though.  Your module might have been unloaded already.
+	 *
+	 * @return Returns singleton instance, loading the module on demand if needed
+	 */
+	static inline IVolumetrics& Get()
+	{
+		return FModuleManager::LoadModuleChecked< IVolumetrics >( "Volumetrics" );
+	}
+
+	/**
+	 * Checks to see if this module is loaded and ready.  It is only valid to call Get() if IsAvailable() returns true.
+	 *
+	 * @return True if the module is loaded and ready to use
+	 */
+	static inline bool IsAvailable()
+	{
+		return FModuleManager::Get().IsModuleLoaded( "Volumetrics" );
+	}
+};
+
diff --git a/Engine/Plugins/Experimental/Volumetrics/Source/Volumetrics/Volumetrics.Build.cs b/Engine/Plugins/Experimental/Volumetrics/Source/Volumetrics/Volumetrics.Build.cs
new file mode 100644
--- /dev/null
+++ b/Engine/Plugins/Experimental/Volumetrics/Source/Volumetrics/Volumetrics.Build.cs
@@ -0,0 +1,52 @@
+// Copyright 1998-2019 Epic Games, Inc. All Rights Reserved.
+
+namespace UnrealBuildTool.Rules
+{
+	public class Volumetrics : ModuleRules
+	{
+		public Volumetrics(ReadOnlyTargetRules Target) : base(Target)
+		{
+			PublicIncludePaths.AddRange(
+				new string[] {
+					// ... add public include paths required here ...
+				}
+				);
+
+			PrivateIncludePaths.AddRange(
+				new string[] {
+                    "../../../../Source/Runtime/Engine/",
+                    "../../../../Source/Runtime/RHI/",
+					// ... add other private include paths required here ...
+				}
+				);
+            
+            PublicDependencyModuleNames.AddRange(
+				new string[]
+				{
+                    "Core",
+                    "CoreUObject",
+                    "Engine",
+                    "RHI",
+                    "Engine",
+                    "RenderCore",
+					// ... add other public dependencies that you statically link with here ...
+				}
+				);
+
+			PrivateDependencyModuleNames.AddRange(
+				new string[]
+                {
+                    "Projects",
+					// ... add private dependencies that you statically link with here ...
+				}
+				);
+
+			DynamicallyLoadedModuleNames.AddRange(
+				new string[]
+				{
+					// ... add any modules that your module loads dynamically here ...
+				}
+				);
+		}
+	}
+}
diff --git a/Engine/Plugins/Experimental/Volumetrics/Volumetrics.uplugin b/Engine/Plugins/Experimental/Volumetrics/Volumetrics.uplugin
new file mode 100644
--- /dev/null
+++ b/Engine/Plugins/Experimental/Volumetrics/Volumetrics.uplugin
@@ -0,0 +1,24 @@
+{
+	"FileVersion": 3,
+	"Version": 1,
+	"VersionName": "1.0",
+	"FriendlyName": "Volumetrics",
+	"Description": "A library of volume creation and rendering tools using Blueprints.",
+	"Category": "Rendering",
+	"CreatedBy": "Epic Games, Inc.",
+	"CreatedByURL": "http://epicgames.com",
+	"DocsURL": "",
+	"MarketplaceURL": "",
+	"SupportURL": "",
+	"CanContainContent": true,
+	"IsBetaVersion": false,
+	"Installed": false,
+	"Modules" :
+    [
+        {
+            "Name" : "Volumetrics",
+            "Type" : "Runtime",
+            "LoadingPhase" : "PostConfigInit"
+        }
+    ]
+}
\ No newline at end of file
