Volumetrics: Fixes to the Tiling Volume Texture creation BP.

diff --git a/Engine/Plugins/Experimental/Volumetrics/Shaders/Private/VolumetricCloudFunctions.ush b/Engine/Plugins/Experimental/Volumetrics/Shaders/Private/VolumetricCloudFunctions.ush
--- a/Engine/Plugins/Experimental/Volumetrics/Shaders/Private/VolumetricCloudFunctions.ush
+++ b/Engine/Plugins/Experimental/Volumetrics/Shaders/Private/VolumetricCloudFunctions.ush
@@ -115,7 +115,7 @@ struct SkyScene
 	
         for (int i = 0; i < Steps; i++)
         {
-            StepSize *= 1.02;
+            StepSize *= 1.03;
 		
             cursample = SampleCloudLayer(Pos);
 			
