Volumetrics: Set a dependency on the BlueprintMaterialTextureNodes plugin.

diff --git a/Engine/Plugins/Experimental/Volumetrics/Volumetrics.uplugin b/Engine/Plugins/Experimental/Volumetrics/Volumetrics.uplugin
--- a/Engine/Plugins/Experimental/Volumetrics/Volumetrics.uplugin
+++ b/Engine/Plugins/Experimental/Volumetrics/Volumetrics.uplugin
@@ -20,5 +20,11 @@
             "Type" : "Runtime",
             "LoadingPhase" : "PostConfigInit"
         }
-    ]
+    ],
+	"Plugins": [
+		{
+			"Name": "BlueprintMaterialTextureNodes",
+			"Enabled": true
+		}
+	]
 }
\ No newline at end of file
