Removing code from Volumetrics plugin and converted it to a content only plugin.  Removed old volumetric sky raymarch materials from the Planet testmap since it is no longer hooked up.

diff --git a/Engine/Plugins/Experimental/Volumetrics/Source/Volumetrics/Private/Volumetrics.cpp b/Engine/Plugins/Experimental/Volumetrics/Source/Volumetrics/Private/Volumetrics.cpp
deleted file mode 100644
diff --git a/Engine/Plugins/Experimental/Volumetrics/Source/Volumetrics/Public/IVolumetrics.h b/Engine/Plugins/Experimental/Volumetrics/Source/Volumetrics/Public/IVolumetrics.h
deleted file mode 100644
diff --git a/Engine/Plugins/Experimental/Volumetrics/Source/Volumetrics/Volumetrics.Build.cs b/Engine/Plugins/Experimental/Volumetrics/Source/Volumetrics/Volumetrics.Build.cs
deleted file mode 100644
diff --git a/Engine/Plugins/Experimental/Volumetrics/Volumetrics.uplugin b/Engine/Plugins/Experimental/Volumetrics/Volumetrics.uplugin
--- a/Engine/Plugins/Experimental/Volumetrics/Volumetrics.uplugin
+++ b/Engine/Plugins/Experimental/Volumetrics/Volumetrics.uplugin
@@ -13,14 +13,6 @@
 	"CanContainContent": true,
 	"IsBetaVersion": false,
 	"Installed": false,
-	"Modules" :
-    [
-        {
-            "Name" : "Volumetrics",
-            "Type" : "Runtime",
-            "LoadingPhase" : "PostConfigInit"
-        }
-    ],
 	"Plugins": [
 		{
 			"Name": "BlueprintMaterialTextureNodes",
