Add Niagara as an explicit dependency to avoid illegal asset reference errors
#rnx

diff --git a/Engine/Plugins/Experimental/Volumetrics/Volumetrics.uplugin b/Engine/Plugins/Experimental/Volumetrics/Volumetrics.uplugin
--- a/Engine/Plugins/Experimental/Volumetrics/Volumetrics.uplugin
+++ b/Engine/Plugins/Experimental/Volumetrics/Volumetrics.uplugin
@@ -21,6 +21,10 @@
 		{
 			"Name": "Landmass",
 			"Enabled": true
+		},
+		{
+			"Name": "Niagara",
+			"Enabled": true
 		}
 	]
 }
\ No newline at end of file
